DDC: Added DerivedDataTool as a new program to work with DDC

This initial revision can be used to load and save cache replays.

diff --git a/Engine/Source/Programs/DerivedDataTool/DerivedDataTool.Build.cs b/Engine/Source/Programs/DerivedDataTool/DerivedDataTool.Build.cs
new file mode 100644
--- /dev/null
+++ b/Engine/Source/Programs/DerivedDataTool/DerivedDataTool.Build.cs
@@ -0,0 +1,16 @@
+// Copyright Epic Games, Inc. All Rights Reserved.
+
+using UnrealBuildTool;
+
+public class DerivedDataTool : ModuleRules
+{
+	public DerivedDataTool(ReadOnlyTargetRules Target) : base(Target)
+	{
+		PublicIncludePaths.Add("Runtime/Launch/Public");
+		PrivateIncludePaths.Add("Runtime/Launch/Private"); // LaunchEngineLoop.cpp
+
+		PrivateDependencyModuleNames.Add("Core");
+		PrivateDependencyModuleNames.Add("Projects");
+		PrivateDependencyModuleNames.Add("DerivedDataCache");
+	}
+}
diff --git a/Engine/Source/Programs/DerivedDataTool/DerivedDataTool.Target.cs b/Engine/Source/Programs/DerivedDataTool/DerivedDataTool.Target.cs
new file mode 100644
--- /dev/null
+++ b/Engine/Source/Programs/DerivedDataTool/DerivedDataTool.Target.cs
@@ -0,0 +1,32 @@
+// Copyright Epic Games, Inc. All Rights Reserved.
+
+using UnrealBuildTool;
+
+[SupportedPlatforms(UnrealPlatformClass.Desktop)]
+public class DerivedDataToolTarget : TargetRules
+{
+	public DerivedDataToolTarget(TargetInfo Target) : base(Target)
+	{
+		Type = TargetType.Program;
+		IncludeOrderVersion = EngineIncludeOrderVersion.Latest;
+		LinkType = TargetLinkType.Monolithic;
+		LaunchModuleName = "DerivedDataTool";
+
+		bBuildDeveloperTools = false;
+		bCompileAgainstApplicationCore = false;
+		bCompileAgainstCoreUObject = false;
+		bCompileAgainstEngine = false;
+		bCompileICU = false;
+		bForceDisableAutomationTests = true;
+		bIsBuildingConsoleApplication = true;
+		bUseLoggingInShipping = true;
+
+		if (Target.Configuration == UnrealTargetConfiguration.Debug ||
+			Target.Configuration == UnrealTargetConfiguration.Development)
+		{
+			GlobalDefinitions.Add("LLM_ENABLED_IN_CONFIG=1");
+			GlobalDefinitions.Add("UE_MEMORY_TAGS_TRACE_ENABLED=1");
+			GlobalDefinitions.Add("UE_TRACE_ENABLED=1");
+		}
+	}
+}
diff --git a/Engine/Source/Programs/DerivedDataTool/Private/DerivedDataTool.cpp b/Engine/Source/Programs/DerivedDataTool/Private/DerivedDataTool.cpp
new file mode 100644
--- /dev/null
+++ b/Engine/Source/Programs/DerivedDataTool/Private/DerivedDataTool.cpp
@@ -0,0 +1,40 @@
+// Copyright Epic Games, Inc. All Rights Reserved.
+
+#include "RequiredProgramMainCPPInclude.h"
+
+#include "CoreGlobals.h"
+#include "DerivedDataCache.h"
+#include "HAL/LowLevelMemTracker.h"
+#include "Misc/ScopeExit.h"
+#include "Modules/ModuleManager.h"
+
+DEFINE_LOG_CATEGORY_STATIC(LogDerivedDataTool, Log, All);
+
+IMPLEMENT_APPLICATION(DerivedDataTool, "DerivedDataTool");
+
+INT32_MAIN_INT32_ARGC_TCHAR_ARGV()
+{
+	using namespace UE::DerivedData;
+
+	const FTaskTagScope Scope(ETaskTag::EGameThread);
+
+	if (int32 Ret = GEngineLoop.PreInit(ArgC, ArgV))
+	{
+		return Ret;
+	}
+
+	ON_SCOPE_EXIT
+	{
+	#if ENABLE_LOW_LEVEL_MEM_TRACKER
+		FLowLevelMemTracker::Get().UpdateStatsPerFrame();
+	#endif
+		RequestEngineExit(TEXT("Exiting"));
+		FEngineLoop::AppPreExit();
+		FModuleManager::Get().UnloadModulesAtShutdown();
+		FEngineLoop::AppExit();
+	};
+
+	GetCache();
+
+	return 0;
+}
