DerivedDataTool: Added support for RaceDetector

#rnx
#Synced-CL 32500503

diff --git a/Engine/Source/Programs/DerivedDataTool/Private/DerivedDataTool.cpp b/Engine/Source/Programs/DerivedDataTool/Private/DerivedDataTool.cpp
--- a/Engine/Source/Programs/DerivedDataTool/Private/DerivedDataTool.cpp
+++ b/Engine/Source/Programs/DerivedDataTool/Private/DerivedDataTool.cpp
@@ -8,6 +8,7 @@
 #include "Misc/ScopeExit.h"
 #include "Modules/ModuleManager.h"
 #include "ProjectUtilities.h"
+#include "Sanitizer/RaceDetector.h"
 
 DEFINE_LOG_CATEGORY_STATIC(LogDerivedDataTool, Log, All);
 
@@ -15,12 +16,17 @@ IMPLEMENT_APPLICATION(DerivedDataTool, "DerivedDataTool");
 
 INT32_MAIN_INT32_ARGC_TCHAR_ARGV()
 {
+	using namespace UE;
 	using namespace UE::DerivedData;
 
+#if USING_INSTRUMENTATION
+	Sanitizer::RaceDetector::Initialize();
+#endif
+
 	const FTaskTagScope Scope(ETaskTag::EGameThread);
 
 	// Allows this program to accept a project argument on the commandline and use project-specific config
-	UE::ProjectUtilities::ParseProjectDirFromCommandline(ArgC, ArgV);
+	ProjectUtilities::ParseProjectDirFromCommandline(ArgC, ArgV);
 
 	if (int32 Ret = GEngineLoop.PreInit(ArgC, ArgV))
 	{
@@ -36,6 +42,9 @@ INT32_MAIN_INT32_ARGC_TCHAR_ARGV()
 		FEngineLoop::AppPreExit();
 		FModuleManager::Get().UnloadModulesAtShutdown();
 		FEngineLoop::AppExit();
+	#if USING_INSTRUMENTATION
+		Sanitizer::RaceDetector::Shutdown();
+	#endif
 	};
 
 	GetCache();
