DerivedDataTool: Added a CreateCache command to start phasing out creation of the default cache graph

diff --git a/Engine/Source/Programs/DerivedDataTool/Private/DerivedDataTool.cpp b/Engine/Source/Programs/DerivedDataTool/Private/DerivedDataTool.cpp
--- a/Engine/Source/Programs/DerivedDataTool/Private/DerivedDataTool.cpp
+++ b/Engine/Source/Programs/DerivedDataTool/Private/DerivedDataTool.cpp
@@ -2,22 +2,151 @@
 
 #include "RequiredProgramMainCPPInclude.h"
 
+#include "Algo/Find.h"
+#include "Containers/UnrealString.h"
 #include "CoreGlobals.h"
 #include "DerivedDataCache.h"
+#include "DerivedDataToolCommand.h"
 #include "HAL/LowLevelMemTracker.h"
+#include "Logging/LogMacros.h"
+#include "Misc/Parse.h"
 #include "Misc/ScopeExit.h"
+#include "Misc/StringBuilder.h"
 #include "Modules/ModuleManager.h"
 #include "ProjectUtilities.h"
 #include "Sanitizer/RaceDetector.h"
+#include "Templates/UnrealTemplate.h"
 
-DEFINE_LOG_CATEGORY_STATIC(LogDerivedDataTool, Log, All);
+DEFINE_LOG_CATEGORY(LogDerivedDataTool);
 
 IMPLEMENT_APPLICATION(DerivedDataTool, "DerivedDataTool");
 
+namespace UE::DerivedData::Tool
+{
+
+///////////////////////////////////////////////////////////////////////////////////////////////////////////////////////
+
+int32 FCommand::Execute(const TCHAR* Command) const
+{
+	TStringBuilder<1024> Tokens;
+	TStringBuilder<1024> Options;
+
+	FString Token;
+	for (const TCHAR* End = Command; FParse::Token(End, Token, /*bUseEscape*/ false); Command = End)
+	{
+		FStringBuilderBase& Out = (**Token == TEXT('-')) ? Options : Tokens;
+		Out.Append(Command, UE_PTRDIFF_TO_INT32(End - Command));
+	}
+
+	TArray<const FCommand*> Stack;
+	return Execute(*Tokens, *Options, Stack);
+}
+
+int32 FCommand::Execute(const TCHAR* Tokens, const TCHAR* Options, TArray<const FCommand*>& Stack) const
+{
+	Stack.Push(this);
+	ON_SCOPE_EXIT
+	{
+		Stack.Pop(EAllowShrinking::No);
+	};
+
+	if (FParse::Command(&Tokens, TEXT("Help")))
+	{
+		return Help(Stack);
+	}
+	else if (const FCommand* Command = Algo::FindByPredicate(Commands,
+		[&Tokens](const FCommand& C) { return FParse::Command(&Tokens, C.Name); }))
+	{
+		return Command->Execute(Tokens, Options, Stack);
+	}
+	else if (Function)
+	{
+		return Function(Tokens, Options);
+	}
+	else
+	{
+		GetCache();
+		return Help(Stack);
+	}
+}
+
+int32 FCommand::Help(TConstArrayView<const FCommand*> Stack) const
+{
+	TGuardValue NoTime(GPrintLogTimes, ELogTimes::None);
+	TGuardValue NoCategory(GPrintLogCategory, false);
+	TGuardValue NoVerbosity(GPrintLogVerbosity, false);
+
+	TUtf8StringBuilder<256> UsageLine;
+	UsageLine.Append(ANSITEXTVIEW("Usage: "));
+	for (const FCommand* Command : Stack)
+	{
+		UsageLine.Append(Command->Name).AppendChar(' ');
+	}
+	const int32 BaseUsageLen = UsageLine.Len();
+	UsageLine.Append(Usage);
+	UE_LOGF(LogDerivedDataTool, Display, "%s", *UsageLine);
+
+	if (Description)
+	{
+		UE_LOGF(LogDerivedDataTool, Display, "");
+		UE_LOGF(LogDerivedDataTool, Display, "%s", Description);
+	}
+
+	if (!Commands.IsEmpty())
+	{
+		UE_LOGF(LogDerivedDataTool, Display, "");
+		UE_LOGF(LogDerivedDataTool, Display, "Commands:");
+		UE_LOGF(LogDerivedDataTool, Display, "");
+		for (const FCommand& Command : Commands)
+		{
+			UsageLine.RemoveSuffix(UsageLine.Len() - BaseUsageLen);
+			UsageLine.Append(Command.Name).AppendChar(' ').Append(Command.Usage);
+			UE_LOGF(LogDerivedDataTool, Display, "    %s", *UsageLine);
+		}
+	}
+
+	for (const FCommand* Command : Stack)
+	{
+		if (Command->Switches.IsEmpty())
+		{
+			continue;
+		}
+		UE_LOGF(LogDerivedDataTool, Display, "");
+		UE_LOGF(LogDerivedDataTool, Display, "Switches:");
+		UE_LOGF(LogDerivedDataTool, Display, "");
+		for (const FSwitch& Switch : Command->Switches)
+		{
+			UE_LOGF(LogDerivedDataTool, Display, "    -%-30s    %s", Switch.Name, Switch.Description ? Switch.Description : "");
+		}
+	}
+
+	GLog->Flush();
+
+	return 1;
+}
+
+///////////////////////////////////////////////////////////////////////////////////////////////////////////////////////
+
+FCommand MakeCreateCacheCommand()
+{
+	return FCommand(TEXT("CreateCache"))
+		.SetDescription("Create the cache graph and exit. Use to validate config or load replays.")
+		.AddSwitch("DDC=<GraphNameOrConfig>", "The graph name or config to create the cache graph from.")
+		.AddSwitch("DDC-ReplayLoad=<Path>", "Replay to load and replay against the graph. Supports multiple replays.")
+		.OnExecute([](auto, auto)
+		{
+			GetCache();
+			return 0;
+		});
+}
+
+} // UE::DerivedData::Tool
+
 INT32_MAIN_INT32_ARGC_TCHAR_ARGV()
 {
 	using namespace UE;
 	using namespace UE::DerivedData;
+	using namespace UE::DerivedData::Tool;
 
 #if USING_INSTRUMENTATION
 	Sanitizer::RaceDetector::Initialize();
@@ -35,9 +164,7 @@ INT32_MAIN_INT32_ARGC_TCHAR_ARGV()
 
 	ON_SCOPE_EXIT
 	{
-	#if ENABLE_LOW_LEVEL_MEM_TRACKER
-		FLowLevelMemTracker::Get().UpdateStatsPerFrame();
-	#endif
+		LLM(FLowLevelMemTracker::Get().UpdateStatsPerFrame());
 		RequestEngineExit(TEXT("Exiting"));
 		FEngineLoop::AppPreExit();
 		FModuleManager::Get().UnloadModulesAtShutdown();
@@ -47,7 +174,20 @@ INT32_MAIN_INT32_ARGC_TCHAR_ARGV()
 	#endif
 	};
 
-	GetCache();
+	const TCHAR* Command = FCommandLine::Get();
+
+	// Skip the optional project argument.
+	{
+		const TCHAR* RemainingCommand = Command;
+		if (FString ProjectToken; FParse::Token(RemainingCommand, ProjectToken, /*bUseEscape*/ true) &&
+			ProjectToken.EndsWith(FProjectDescriptor::GetExtension()))
+		{
+			Command = RemainingCommand;
+		}
+	}
 
-	return 0;
+	FCommand RootCommand = FCommand(TEXT("DerivedDataTool"))
+		.AddCommand(MakeCreateCacheCommand())
+	;
+	return RootCommand.Execute(Command);
 }
diff --git a/Engine/Source/Programs/DerivedDataTool/Private/DerivedDataToolCommand.h b/Engine/Source/Programs/DerivedDataTool/Private/DerivedDataToolCommand.h
new file mode 100644
--- /dev/null
+++ b/Engine/Source/Programs/DerivedDataTool/Private/DerivedDataToolCommand.h
@@ -0,0 +1,74 @@
+// Copyright Epic Games, Inc. All Rights Reserved.
+
+#pragma once
+
+#include "Containers/Array.h"
+#include "Logging/LogCategory.h"
+
+DECLARE_LOG_CATEGORY_EXTERN(LogDerivedDataTool, Log, All);
+
+namespace UE::DerivedData::Tool
+{
+
+class FCommand final
+{
+public:
+	using FFunction = int32 (const TCHAR* Tokens, const TCHAR* Options);
+
+	explicit FCommand(const TCHAR* InName)
+		: Name(InName)
+	{
+	}
+
+	constexpr FCommand&& SetUsage(const ANSICHAR* InUsage) &&
+	{
+		Usage = InUsage;
+		return (FCommand&&)*this;
+	}
+
+	constexpr FCommand&& SetDescription(const ANSICHAR* InDescription) &&
+	{
+		Description = InDescription;
+		return (FCommand&&)*this;
+	}
+
+	FCommand&& AddSwitch(const ANSICHAR* InSwitch, const ANSICHAR* InDescription = nullptr) &&
+	{
+		Switches.Add({InSwitch, InDescription});
+		return (FCommand&&)*this;
+	}
+
+	FCommand&& AddCommand(FCommand&& Command) &&
+	{
+		Commands.Emplace(MoveTemp(Command));
+		return (FCommand&&)*this;
+	}
+
+	constexpr FCommand&& OnExecute(FFunction* InFunction) &&
+	{
+		Function = InFunction;
+		return (FCommand&&)*this;
+	}
+
+	int32 Execute(const TCHAR* Command) const;
+
+private:
+	int32 Execute(const TCHAR* Tokens, const TCHAR* Options, TArray<const FCommand*>& Stack) const;
+
+	int32 Help(TConstArrayView<const FCommand*> Stack) const;
+
+	struct FSwitch
+	{
+		const ANSICHAR* Name = nullptr;
+		const ANSICHAR* Description = nullptr;
+	};
+
+	const TCHAR* Name = nullptr;
+	const ANSICHAR* Usage = nullptr;
+	const ANSICHAR* Description = nullptr;
+	TArray<FSwitch> Switches;
+	TArray<FCommand> Commands;
+	FFunction* Function = nullptr;
+};
+
+} // UE::DerivedData::Tool
