DerivedDataTool: Allow the Copy command to be called with no targets

It will load replays and summarize the keys to be scanned before exiting.

#rnx
#Synced-CL 49572071

diff --git a/Engine/Source/Programs/DerivedDataTool/DerivedDataTool.Build.cs b/Engine/Source/Programs/DerivedDataTool/DerivedDataTool.Build.cs
--- a/Engine/Source/Programs/DerivedDataTool/DerivedDataTool.Build.cs
+++ b/Engine/Source/Programs/DerivedDataTool/DerivedDataTool.Build.cs
@@ -15,5 +15,7 @@ public class DerivedDataTool : ModuleRules
 		PrivateDependencyModuleNames.Add("DerivedDataCache");
 		PrivateDependencyModuleNames.Add("DesktopPlatform");
 		PrivateDependencyModuleNames.Add("Projects");
+
+		CppCompileWarningSettings.UnsafeTypeCastWarningLevel = WarningLevel.Error;
 	}
 }
diff --git a/Engine/Source/Programs/DerivedDataTool/Private/DerivedDataCopyCommand.cpp b/Engine/Source/Programs/DerivedDataTool/Private/DerivedDataCopyCommand.cpp
--- a/Engine/Source/Programs/DerivedDataTool/Private/DerivedDataCopyCommand.cpp
+++ b/Engine/Source/Programs/DerivedDataTool/Private/DerivedDataCopyCommand.cpp
@@ -520,8 +520,6 @@ struct FCopyParams
 		"Using salt -FilterSalt={FilterSalt} to filter cache keys.", OutParams.Filter.GetSalt());
 
 	FParse::Value(Options, TEXT("-ValueMemoryMB="), OutParams.ValueMemoryMB);
-	UE_LOGFMT(LogDerivedDataTool, Display,
-		"Limiting memory for copied values to {ValueMemoryMB} MiB.", OutParams.ValueMemoryMB);
 
 	OutParams.bDebugSkipPut = FParse::Param(Options, TEXT("-DebugSkipPut"));
 	OutParams.bPreview = FParse::Param(Options, TEXT("-Preview"));
@@ -536,7 +534,7 @@ struct FCopyParams
 		bOk = false;
 	}
 
-	if (OutParams.SourceNameOrConfig.IsEmpty() && (!OutParams.bPreview || !OutParams.Keys.IsEmpty()))
+	if (OutParams.SourceNameOrConfig.IsEmpty() && !OutParams.TargetConfigs.IsEmpty() && (!OutParams.bPreview || !OutParams.Keys.IsEmpty()))
 	{
 		UE_LOGFMT(LogDerivedDataTool, Error, "Argument '-Source' is required for non-preview copies or when individual keys are provided.");
 		bOk = false;
@@ -1918,6 +1916,11 @@ int32 Copy(const TCHAR* Tokens, const TCHAR* Options)
 		return 1;
 	}
 
+	if (Context.Params.TargetConfigs.IsEmpty())
+	{
+		return 0;
+	}
+
 	if (!Context.Params.bPreview || !Context.Params.Keys.IsEmpty())
 	{
 		if (!CreateCopySource(Context))
@@ -1932,6 +1935,7 @@ int32 Copy(const TCHAR* Tokens, const TCHAR* Options)
 	}
 
 	Context.MemoryLimit.MaxUsedMemory = uint64(Context.Params.ValueMemoryMB) * 1'048'576;
+	UE_LOGFMT(LogDerivedDataTool, Display, "Limiting memory for copied values to {ValueMemoryMB} MiB.", Context.Params.ValueMemoryMB);
 
 	ConfigureApis(Context);
 
