DerivedDataTool: Optimized replay loading by using FProbingHashTable and TConcurrentChunkedSparseArray

The net reduction in replay loading time when loading 10 replays in parallel is ~48%.
Replacing TMap with FProbingHashTable reduced replay load time by ~27%.
Replacing TChunkedArray with TConcurrentChunkedSparseArray reduced replay load time by ~19%.
Sharding the FProbingHashTable and associated mutexes reduced replay load time by ~18%.
Changing BucketsMutex from FMutex to FSharedMutex reduced FindOrCreateBucket cost by ~71%.

#rnx
#Synced-CL 49572071

diff --git a/Engine/Source/Programs/DerivedDataTool/Private/DerivedDataCopyCommand.cpp b/Engine/Source/Programs/DerivedDataTool/Private/DerivedDataCopyCommand.cpp
--- a/Engine/Source/Programs/DerivedDataTool/Private/DerivedDataCopyCommand.cpp
+++ b/Engine/Source/Programs/DerivedDataTool/Private/DerivedDataCopyCommand.cpp
@@ -10,8 +10,10 @@
 #include "Async/SharedLock.h"
 #include "Async/SharedMutex.h"
 #include "Async/UniqueLock.h"
-#include "Containers/ChunkedArray.h"
+#include "Concepts/Integral.h"
+#include "Containers/ConcurrentChunkedSparseArray.h"
 #include "Containers/Map.h"
+#include "Containers/ProbingHashTable.h"
 #include "DerivedDataCache.h"
 #include "DerivedDataCacheKeyFilter.h"
 #include "DerivedDataCacheReplay.h"
@@ -23,6 +25,7 @@
 #include "HAL/PlatformTime.h"
 #include "Logging/LogMacros.h"
 #include "Logging/StructuredLog.h"
+#include "Misc/AlignedElement.h"
 #include "Misc/Parse.h"
 #include "Misc/ScopeExit.h"
 #include "Misc/StringBuilder.h"
@@ -46,6 +49,9 @@ constexpr static int32 PutRequestBatchSize = 16;
 // Flags for targets are stored in uint32.
 constexpr static int32 MaxCopyTargets = 32;
 
+// Buckets distribute keys into shards to reduce contention.
+constexpr static uint32 BucketShardCount = 64;
+
 int32 Copy(const TCHAR* Tokens, const TCHAR* Options);
 
 FCommand MakeCopyCommand()
@@ -81,6 +87,28 @@ FCommand MakeCopyCommand()
 
 ///////////////////////////////////////////////////////////////////////////////////////////////////////////////////////
 
+template <CIntegral T>
+class TRelaxedAtomicInt
+{
+public:
+	inline TRelaxedAtomicInt(T N) : Value(N) {}
+ 
+	inline operator T() const { return Value.load(std::memory_order_relaxed); }
+
+	inline TRelaxedAtomicInt& operator=(T N) { Value.store(N, std::memory_order_relaxed); return *this; }
+
+	inline void operator+=(T N) { Value.fetch_add(N, std::memory_order_relaxed); }
+	inline void operator-=(T N) { Value.fetch_sub(N, std::memory_order_relaxed); }
+
+	inline T operator++() { return Value.fetch_add(1, std::memory_order_relaxed) + 1;}
+	inline T operator--() { return Value.fetch_sub(1, std::memory_order_relaxed) - 1; }
+
+private:
+	std::atomic<T> Value = 0;
+};
+
+///////////////////////////////////////////////////////////////////////////////////////////////////////////////////////
+
 class FCopyRequestBatch;
 class FCopyRequestBatchPool;
 
@@ -89,11 +117,11 @@ struct FCopyRequestBatchApi
 	// CanDispatch returns the request count when ready to dispatch, otherwise -1.
 	// May reserve memory but may not release memory because release can trigger dispatch.
 	// Must be trivial because this is called within the pool lock.
-	int32 (*CanDispatch)(void*, TConstArrayView<int32>) = nullptr;
+	int32 (*CanDispatch)(void*, TConstArrayView<uint32>) = nullptr;
 	// Dispatch is called outside of the pool lock after CanDispatch returns a non-negative request count.
-	void (*Dispatch)(void*, TConstArrayView<int32>, IRequestOwner&, const FCopyRequestBatch&) = nullptr;
+	void (*Dispatch)(void*, TConstArrayView<uint32>, IRequestOwner&, const FCopyRequestBatch&) = nullptr;
 	// Complete is called outside of the pool lock when the last request in the batch completes.
-	void (*Complete)(void*, TConstArrayView<int32>) = nullptr;
+	void (*Complete)(void*, TConstArrayView<uint32>) = nullptr;
 	// Stall is called outside of the pool lock after CanDispatch returns a negative request count.
 	void (*Stall)(void*) = nullptr;
 	// Context passed to each of the functions.
@@ -103,7 +131,7 @@ struct FCopyRequestBatchApi
 class FCopyRequestBatch
 {
 public:
-	[[nodiscard]] inline bool AddIndex(int32 Index);
+	[[nodiscard]] inline bool AddIndex(uint32 Index);
 
 	inline void Complete() const;
 
@@ -123,7 +151,7 @@ private:
 	int32 OwnerIndex = -1;
 	int32 IndexCapacity = 0;
 	int32 IndexCount = 0;
-	int32 IndexArray[0];
+	uint32 IndexArray[0];
 
 	friend FCopyRequestBatchPool;
 };
@@ -221,7 +249,7 @@ public:
 			}
 
 			// Query whether the batch is ready to dispatch.
-			const TConstArrayView<int32> IndexBatch(Batch->IndexArray, Batch->IndexCount);
+			const TConstArrayView<uint32> IndexBatch(Batch->IndexArray, Batch->IndexCount);
 			const int32 RequestCount = Batch->Api->CanDispatch(Batch->Api->Context, IndexBatch);
 
 			// Pause any further dispatching if the batch is not ready yet.
@@ -304,7 +332,7 @@ private:
 
 inline FCopyRequestBatch* FCopyRequestBatch::New(TNotNull<FCopyRequestBatchPool*> InPool)
 {
-	const SIZE_T Size = sizeof(FCopyRequestBatch) + sizeof(int32) * InPool->BatchSize;
+	const SIZE_T Size = sizeof(FCopyRequestBatch) + sizeof(uint32) * InPool->BatchSize;
 	return new(FMemory::Malloc(Size, alignof(FCopyRequestBatch))) FCopyRequestBatch(InPool);
 }
 
@@ -318,10 +346,10 @@ inline FCopyRequestBatch::FCopyRequestBatch(TNotNull<FCopyRequestBatchPool*> InP
 	: Pool(InPool)
 	, IndexCapacity(InPool->BatchSize)
 {
-	FMemory::Memzero(IndexArray, IndexCapacity * sizeof(int32));
+	FMemory::Memzero(IndexArray, IndexCapacity * sizeof(uint32));
 }
 
-inline bool FCopyRequestBatch::AddIndex(int32 Index)
+inline bool FCopyRequestBatch::AddIndex(uint32 Index)
 {
 	IndexArray[IndexCount] = Index;
 	return ++IndexCount == IndexCapacity;
@@ -355,7 +383,7 @@ public:
 		check(!Batch);
 	}
 
-	inline void Queue(int32 Index)
+	inline void Queue(uint32 Index)
 	{
 		FCopyRequestBatch* LocalBatch = nullptr;
 		{
@@ -548,22 +576,22 @@ struct FCopyParams
 struct FCopyLoadSummary
 {
 	/** Unique keys for records. */
-	int32 RecordKeyCount = 0;
+	TRelaxedAtomicInt<int32> RecordKeyCount = 0;
 	/** Unique keys for values. */
-	int32 ValueKeyCount = 0;
+	TRelaxedAtomicInt<int32> ValueKeyCount = 0;
 	/** Unique key+record count. */
-	int32 RecordCount = 0;
+	TRelaxedAtomicInt<int32> RecordCount = 0;
 	/** Unique key+record+value count. (Number of values in unique records.) */
-	int32 ValueInRecordCount = 0;
+	TRelaxedAtomicInt<int32> ValueInRecordCount = 0;
 	/** Unique key+value count. */
-	int32 ValueCount = 0;
+	TRelaxedAtomicInt<int32> ValueCount = 0;
 	/** Total record count before de-duplication. */
-	int32 TotalRecordCount = 0;
+	TRelaxedAtomicInt<int32> TotalRecordCount = 0;
 	/** Total value count before de-duplication by key+value. */
-	int32 TotalValueCount = 0;
+	TRelaxedAtomicInt<int32> TotalValueCount = 0;
 
 	/** Total raw size of values referenced by unique keys. */
-	uint64 TotalRawSize = 0;
+	TRelaxedAtomicInt<uint64> TotalRawSize = 0;
 };
 
 inline FCopyLoadSummary& operator+=(FCopyLoadSummary& Out, const FCopyLoadSummary& In)
@@ -659,9 +687,9 @@ struct FCopyRecord
 	FSharedString Name;
 	FCacheBucket Bucket;
 	FIoHash KeyHash;
-	int32 ValueIndex = 0;
-	int32 ValueCount = 0;
-	int32 NextIndex = -1;
+	uint32 ValueIndex = 0;
+	uint32 ValueCount = 0;
+	uint32 NextIndex = MAX_uint32;
 	FCopyFlags Flags;
 	FOptionalCacheRecord Data;
 };
@@ -680,7 +708,7 @@ struct FCopyValue
 	FSharedString Name;
 	FCacheBucket Bucket;
 	FIoHash KeyHash;
-	int32 NextIndex = -1;
+	uint32 NextIndex = MAX_uint32;
 	uint64 RawSize = 0;
 	FIoHash RawHash;
 	FCopyFlags Flags;
@@ -690,27 +718,24 @@ static_assert(sizeof(FCopyValue) == 96); // Padding == 4
 
 struct FCopyBucket
 {
-	TMap<FIoHash, int32> RecordKeys;
-	TMap<FIoHash, int32> ValueKeys;
+	Core::TAlignedElement<FProbingHashTable, PLATFORM_CACHE_LINE_SIZE> RecordKeys[BucketShardCount];
+	Core::TAlignedElement<FProbingHashTable, PLATFORM_CACHE_LINE_SIZE> ValueKeys[BucketShardCount];
+	TRelaxedAtomicInt<int32> RecordKeyCount = 0;
+	TRelaxedAtomicInt<int32> ValueKeyCount = 0;
 	FCopyLoadSummary LoadSummary;
-	mutable FMutex Mutex;
+
+	inline static uint32 MapHashToShard(const FIoHash& Hash)
+	{
+		return reinterpret_cast<const uint32*>(&Hash)[1] % BucketShardCount;
+	}
 };
 
 struct FCopyContext
 {
-	struct alignas(PLATFORM_CACHE_LINE_SIZE)
-	{
-		TMap<FCacheBucket, TUniquePtr<FCopyBucket>> Buckets;
-	};
-	struct alignas(PLATFORM_CACHE_LINE_SIZE)
-	{
-		TChunkedArray<FCopyRecord> RecordTable;
-		TChunkedArray<FCopyValueInRecord> ValueInRecordTable;
-	};
-	struct alignas(PLATFORM_CACHE_LINE_SIZE)
-	{
-		TChunkedArray<FCopyValue> ValueTable;
-	};
+	alignas(PLATFORM_CACHE_LINE_SIZE) TMap<FCacheBucket, TUniquePtr<FCopyBucket>> Buckets;
+	alignas(PLATFORM_CACHE_LINE_SIZE) TConcurrentChunkedSparseArray<FCopyRecord, 32 * 1024> RecordTable;
+	alignas(PLATFORM_CACHE_LINE_SIZE) TConcurrentChunkedSparseArray<FCopyValueInRecord, 128 * 1024> ValueInRecordTable;
+	alignas(PLATFORM_CACHE_LINE_SIZE) TConcurrentChunkedSparseArray<FCopyValue, 32 * 1024> ValueTable;
 
 	TUniquePtr<ICache> Source;
 	TArray<TUniquePtr<ICache>, TFixedAllocator<MaxCopyTargets>> Targets;
@@ -738,19 +763,19 @@ struct FCopyContext
 	FCopyProgress Progress;
 	FCopyParams Params;
 
-	[[nodiscard]] int32 Find(int32 HeadRecordIndex, TConstArrayView<FValueWithId> Values) const;
-	[[nodiscard]] int32 Find(int32 HeadValueIndex, const FValue& Value) const;
+	[[nodiscard]] uint32 Find(uint32 HeadRecordIndex, TConstArrayView<FValueWithId> Values) const;
+	[[nodiscard]] uint32 Find(uint32 HeadValueIndex, const FValue& Value) const;
 };
 
-int32 FCopyContext::Find(int32 HeadRecordIndex, TConstArrayView<FValueWithId> SearchValues) const
+uint32 FCopyContext::Find(uint32 HeadRecordIndex, TConstArrayView<FValueWithId> SearchValues) const
 {
-	for (int32 RecordIndex = HeadRecordIndex; RecordIndex != -1;)
+	for (uint32 RecordIndex = HeadRecordIndex; RecordIndex != MAX_uint32;)
 	{
 		const FCopyRecord& Record = RecordTable[RecordIndex];
-		if (Record.ValueCount == SearchValues.Num())
+		if (Record.ValueCount == (uint32)SearchValues.Num())
 		{
 			bool bEqual = true;
-			int32 ValueIndex = Record.ValueIndex;
+			uint32 ValueIndex = Record.ValueIndex;
 			for (const FValueWithId& CacheValue : SearchValues)
 			{
 				const FCopyValueInRecord& Value = ValueInRecordTable[ValueIndex++];
@@ -766,12 +791,12 @@ int32 FCopyContext::Find(int32 HeadRecordIndex, TConstArrayView<FValueWithId> Se
 		}
 		RecordIndex = Record.NextIndex;
 	}
-	return -1;
+	return MAX_uint32;
 }
 
-int32 FCopyContext::Find(int32 HeadValueIndex, const FValue& SearchValue) const
+uint32 FCopyContext::Find(uint32 HeadValueIndex, const FValue& SearchValue) const
 {
-	for (int32 ValueIndex = HeadValueIndex; ValueIndex != -1;)
+	for (uint32 ValueIndex = HeadValueIndex; ValueIndex != MAX_uint32;)
 	{
 		const FCopyValue& Value = ValueTable[ValueIndex];
 		if (Value.RawSize == SearchValue.GetRawSize() &&
@@ -781,7 +806,7 @@ int32 FCopyContext::Find(int32 HeadValueIndex, const FValue& SearchValue) const
 		}
 		ValueIndex = Value.NextIndex;
 	}
-	return -1;
+	return MAX_uint32;
 }
 
 ///////////////////////////////////////////////////////////////////////////////////////////////////////////////////////
@@ -839,7 +864,7 @@ inline static FCachePutValueRequest MakePutRequest(const FCopyValue& Value, ECac
 [[nodiscard]] inline static uint64 MeasureRawSize(const FCopyContext& Context, const FCopyRecord& Record)
 {
 	uint64 RawSize = 0;
-	for (int32 Index = Record.ValueIndex, End = Record.ValueCount + Index; Index < End; ++Index)
+	for (uint32 Index = Record.ValueIndex, End = Record.ValueCount + Index; Index < End; ++Index)
 	{
 		RawSize += Context.ValueInRecordTable[Index].RawSize;
 	}
@@ -951,89 +976,142 @@ private:
 	void AddRecord(const FSharedString& Name, const FCacheKey& Key, TConstArrayView<FValueWithId> Values)
 	{
 		FCopyBucket& Bucket = FindOrCreateBucket(Key.Bucket);
-		TUniqueLock BucketLock(Bucket.Mutex);
-		int32& HeadRecordIndex = Bucket.RecordKeys.FindOrAdd(Key.Hash, -1);
-		++Bucket.LoadSummary.TotalRecordCount;
 
-		// Check for an existing copy of this record.
+		const uint32 Shard = Bucket.MapHashToShard(Key.Hash);
+		TUniqueLock<FMutex> BucketLock(BucketKeysMutex[Shard]);
+		FProbingHashTable& RecordKeys = Bucket.RecordKeys[Shard];
+
+		uint32 HeadRecordIndex = MAX_uint32;
+
+		const uint32 KeyHash = GetTypeHash(Key.Hash);
+		FProbingHashTable::FSlotId SlotId;
 		{
-			TSharedLock Lock(RecordTableMutex);
-			if (int32 ExistingRecordIndex = Context.Find(HeadRecordIndex, Values); ExistingRecordIndex != -1)
+			uint32 Index = 0;
+			for (SlotId = RecordKeys.FindFirst(KeyHash, Index); SlotId; SlotId = RecordKeys.FindNext(SlotId, KeyHash, Index))
 			{
-				FCopyRecord& ExistingRecord = Context.RecordTable[ExistingRecordIndex];
-				if (Name < ExistingRecord.Name)
+				if (Context.RecordTable[Index].KeyHash == Key.Hash)
 				{
-					ExistingRecord.Name = Name;
+					HeadRecordIndex = Index;
+					break;
 				}
-				return;
 			}
 		}
 
-		TUniqueLock Lock(RecordTableMutex);
+		++Bucket.LoadSummary.TotalRecordCount;
+
+		// Check for an existing copy of this record.
+		if (uint32 ExistingRecordIndex = Context.Find(HeadRecordIndex, Values); ExistingRecordIndex != MAX_uint32)
+		{
+			FCopyRecord& ExistingRecord = Context.RecordTable[ExistingRecordIndex];
+			if (Name < ExistingRecord.Name)
+			{
+				ExistingRecord.Name = Name;
+			}
+			return;
+		}
 
 		// Add the values to the table.
-		int32 FirstValueIndex = -1;
-		for (const FValueWithId& Value : Values)
+		const uint32 ValueCount = (uint32)Values.Num();
+		const uint32 FirstValueIndex = Context.ValueInRecordTable.AddUninitialized(ValueCount);
+		for (uint32 ValueIndex = 0; ValueIndex < ValueCount; ++ValueIndex)
 		{
+			const FValueWithId& Value = Values.GetData()[ValueIndex];
 			++Bucket.LoadSummary.ValueInRecordCount;
 			Bucket.LoadSummary.TotalRawSize += Value.GetRawSize();
-			const int32 ValueIndex = Context.ValueInRecordTable.AddElement({Value.GetRawSize(), Value.GetRawHash(), Value.GetId()});
-			if (FirstValueIndex == -1)
-			{
-				FirstValueIndex = ValueIndex;
-			}
+			new(&Context.ValueInRecordTable[FirstValueIndex + ValueIndex]) FCopyValueInRecord{Value.GetRawSize(), Value.GetRawHash(), Value.GetId()};
 		}
 
 		// Add the record to the table.
 		++Bucket.LoadSummary.RecordCount;
-		if (HeadRecordIndex == -1)
+		if (HeadRecordIndex == MAX_uint32)
 		{
 			++Bucket.LoadSummary.RecordKeyCount;
+			++Bucket.RecordKeyCount;
+		}
+		const uint32 NewRecordIndex = Context.RecordTable.AddUninitialized();
+		new(&Context.RecordTable[NewRecordIndex]) FCopyRecord{Name, Key.Bucket, Key.Hash, FirstValueIndex, ValueCount, HeadRecordIndex};
+
+		// Add or update the index of the head record for this key.
+		if (SlotId)
+		{
+			RecordKeys[SlotId] = NewRecordIndex;
+		}
+		else
+		{
+			RecordKeys.Add(KeyHash, NewRecordIndex);
 		}
-		int32 NewRecordIndex = Context.RecordTable.Add(1);
-		new(&Context.RecordTable[NewRecordIndex]) FCopyRecord{Name, Key.Bucket, Key.Hash, FirstValueIndex, Values.Num(), HeadRecordIndex};
-		HeadRecordIndex = NewRecordIndex;
 	}
 
 	void AddValue(const FSharedString& Name, const FCacheKey& Key, const FValue& Value)
 	{
 		FCopyBucket& Bucket = FindOrCreateBucket(Key.Bucket);
-		TUniqueLock BucketLock(Bucket.Mutex);
-		int32& HeadValueIndex = Bucket.ValueKeys.FindOrAdd(Key.Hash, -1);
-		++Bucket.LoadSummary.TotalValueCount;
 
-		// Check for an existing copy of this value.
+		const uint32 Shard = Bucket.MapHashToShard(Key.Hash);
+		TUniqueLock<FMutex> BucketLock(BucketKeysMutex[Shard]);
+		FProbingHashTable& ValueKeys = Bucket.ValueKeys[Shard];
+
+		uint32 HeadValueIndex = MAX_uint32;
+
+		const uint32 KeyHash = GetTypeHash(Key.Hash);
+		FProbingHashTable::FSlotId SlotId;
 		{
-			TSharedLock Lock(ValueTableMutex);
-			if (int32 ExistingValueIndex = Context.Find(HeadValueIndex, Value); ExistingValueIndex != -1)
+			uint32 Index = 0;
+			for (SlotId = ValueKeys.FindFirst(KeyHash, Index); SlotId; SlotId = ValueKeys.FindNext(SlotId, KeyHash, Index))
 			{
-				FCopyValue& ExistingValue = Context.ValueTable[ExistingValueIndex];
-				if (Name < ExistingValue.Name)
+				if (Context.ValueTable[Index].KeyHash == Key.Hash)
 				{
-					ExistingValue.Name = Name;
+					HeadValueIndex = Index;
+					break;
 				}
-				return;
 			}
 		}
 
-		TUniqueLock Lock(ValueTableMutex);
+		++Bucket.LoadSummary.TotalValueCount;
+
+		// Check for an existing copy of this value.
+		if (uint32 ExistingValueIndex = Context.Find(HeadValueIndex, Value); ExistingValueIndex != MAX_uint32)
+		{
+			FCopyValue& ExistingValue = Context.ValueTable[ExistingValueIndex];
+			if (Name < ExistingValue.Name)
+			{
+				ExistingValue.Name = Name;
+			}
+			return;
+		}
 
 		// Add the value to the table.
 		++Bucket.LoadSummary.ValueCount;
-		if (HeadValueIndex == -1)
+		if (HeadValueIndex == MAX_uint32)
 		{
 			++Bucket.LoadSummary.ValueKeyCount;
+			++Bucket.ValueKeyCount;
 		}
 		Bucket.LoadSummary.TotalRawSize += Value.GetRawSize();
-		int32 NewValueIndex = Context.ValueTable.Add(1);
+		const uint32 NewValueIndex = Context.ValueTable.AddUninitialized();
 		new(&Context.ValueTable[NewValueIndex]) FCopyValue{Name, Key.Bucket, Key.Hash, HeadValueIndex, Value.GetRawSize(), Value.GetRawHash()};
-		HeadValueIndex = NewValueIndex;
+
+		// Add or update the index of the head value for this key.
+		if (SlotId)
+		{
+			ValueKeys[SlotId] = NewValueIndex;
+		}
+		else
+		{
+			ValueKeys.Add(KeyHash, NewValueIndex);
+		}
 	}
 
 	FCopyBucket& FindOrCreateBucket(FCacheBucket BucketName)
 	{
+		const uint32 BucketNameHash = GetTypeHash(BucketName);
+
+		if (TSharedLock Lock(BucketsMutex); TUniquePtr<FCopyBucket>* Bucket = Context.Buckets.FindByHash(BucketNameHash, BucketName))
+		{
+			return **Bucket;
+		}
+
 		TUniqueLock Lock(BucketsMutex);
-		TUniquePtr<FCopyBucket>& Bucket = Context.Buckets.FindOrAdd(BucketName);
+		TUniquePtr<FCopyBucket>& Bucket = Context.Buckets.FindOrAddByHash(BucketNameHash, BucketName);
 		if (!Bucket)
 		{
 			Bucket = MakeUnique<FCopyBucket>();
@@ -1046,9 +1124,8 @@ private:
 
 	FCopyContext& Context;
 	const FCacheKeyFilter Filter;
-	alignas(PLATFORM_CACHE_LINE_SIZE) FSharedMutex RecordTableMutex;
-	alignas(PLATFORM_CACHE_LINE_SIZE) FSharedMutex ValueTableMutex;
-	alignas(PLATFORM_CACHE_LINE_SIZE) FMutex BucketsMutex;
+	alignas(PLATFORM_CACHE_LINE_SIZE) FSharedMutex BucketsMutex;
+	Core::TAlignedElement<FMutex, PLATFORM_CACHE_LINE_SIZE> BucketKeysMutex[BucketShardCount];
 };
 
 static void DisplayLoadSummary(const FCopyLoadSummary& Summary, FAnsiStringView BucketPrefix)
@@ -1063,7 +1140,7 @@ static void DisplayLoadSummary(const FCopyLoadSummary& Summary, FAnsiStringView
 
 	const int32 ValueVariantCount = Summary.ValueCount - Summary.ValueKeyCount;
 	UE_CLOGFMT(Summary.ValueKeyCount > 0, LogDerivedDataTool, Display,
-		"{BucketPrefix}Loaded {ValueKeys} value keys, from {TotalValueCount} total values.",
+		"{BucketPrefix}Loaded {ValueKeys} value keys, from {TotalValueCount} value requests.",
 		BucketPrefix, Summary.ValueKeyCount, Summary.TotalValueCount);
 	UE_CLOGFMT(ValueVariantCount > 0, LogDerivedDataTool, Warning,
 		"{BucketPrefix}Found {ValueVariantCount} values that have the same keys with different values.",
@@ -1206,18 +1283,18 @@ public:
 	}
 
 private:
-	static int32 CanDispatch(void* VoidContext, TConstArrayView<int32> IndexBatch)
+	static int32 CanDispatch(void* VoidContext, TConstArrayView<uint32> IndexBatch)
 	{
 		FCopyContext& Context = *(FCopyContext*)VoidContext;
 		return IndexBatch.Num() * Context.Targets.Num();
 	}
 
-	static void Dispatch(void* VoidContext, TConstArrayView<int32> IndexBatch, IRequestOwner& Owner, const FCopyRequestBatch& Batch)
+	static void Dispatch(void* VoidContext, TConstArrayView<uint32> IndexBatch, IRequestOwner& Owner, const FCopyRequestBatch& Batch)
 	{
 		FCopyContext& Context = *(FCopyContext*)VoidContext;
 
 		TArray<RequestType, TFixedAllocator<ScanRequestBatchSize>> Requests;
-		for (int32 Index : IndexBatch)
+		for (uint32 Index : IndexBatch)
 		{
 			const auto& Request = (Context.*TablePtr)[Index];
 			Requests.Add({Request.Name, {Request.Bucket, Request.KeyHash}, ECachePolicy::Query | ECachePolicy::SkipData | ECachePolicy::SkipMeta, uint64(Index)});
@@ -1229,7 +1306,7 @@ private:
 		{
 			for (RequestType& Request : Requests)
 			{
-				const int32 Index = int32(Request.UserData);
+				const uint32 Index = uint32(Request.UserData);
 				Request.UserData = (uint64(TargetIndex) << 32) | Index;
 			}
 			((*Target).*CacheFunctionFor<RequestType>)(Requests, Owner, FOnComplete{Context, Batch});
@@ -1237,7 +1314,7 @@ private:
 		}
 	}
 
-	static void Complete(void* VoidContext, TConstArrayView<int32> IndexBatch)
+	static void Complete(void* VoidContext, TConstArrayView<uint32> IndexBatch)
 	{
 		FCopyContext& Context = *(FCopyContext*)VoidContext;
 
@@ -1257,7 +1334,7 @@ private:
 		void operator()(TCacheResponseFor<RequestType>&& Response) const
 		{
 			const int32 TargetIndex = int32(Response.UserData >> 32);
-			const int32 Index = int32(Response.UserData);
+			const uint32 Index = uint32(Response.UserData);
 
 			auto& Request = (Context.*TablePtr)[Index];
 			if (Response.Status != EStatus::Ok)
@@ -1266,7 +1343,7 @@ private:
 				UE_LOGFMT(LogDerivedDataTool, Verbose, "Value in target '{Target}' missing for {Key} from '{Name}'.",
 					Context.TargetNames[TargetIndex], GetKeyFromResponse(Response), Response.Name);
 			}
-			else if (Context.Find(Index, GetValuesFromResponse(Response)) == -1)
+			else if (Context.Find(Index, GetValuesFromResponse(Response)) == MAX_uint32)
 			{
 				Request.Flags.TargetDiffersMask.fetch_or(uint32(1) << TargetIndex, std::memory_order_relaxed);
 				UE_LOGFMT(LogDerivedDataTool, Display, "Value in target '{Target}' differs for {Key} from '{Name}'.",
@@ -1287,10 +1364,10 @@ public:
 		return {CanDispatch, Dispatch, Complete, Stall, &Context};
 	}
 
-	static void QueueFromScan(FCopyContext& Context, TConstArrayView<int32> IndexBatch)
+	static void QueueFromScan(FCopyContext& Context, TConstArrayView<uint32> IndexBatch)
 	{
 		// Queue every index that has at least one target to copy to.
-		for (int32 Index : IndexBatch)
+		for (uint32 Index : IndexBatch)
 		{
 			auto& Request = (Context.*TablePtr)[Index];
 			const uint32 TargetMissingMask = Request.Flags.TargetMissingMask.load(std::memory_order_relaxed);
@@ -1305,11 +1382,11 @@ public:
 	}
 
 private:
-	static int32 CanDispatch(void* VoidContext, TConstArrayView<int32> IndexBatch)
+	static int32 CanDispatch(void* VoidContext, TConstArrayView<uint32> IndexBatch)
 	{
 		FCopyContext& Context = *(FCopyContext*)VoidContext;
 		// Calculate the raw size and reserve that until the compressed size is known when the requests complete.
-		const auto MeasureIndex = [&Context](int32 Index) -> uint64 { return MeasureRawSize(Context, (Context.*TablePtr)[Index]); };
+		const auto MeasureIndex = [&Context](uint32 Index) -> uint64 { return MeasureRawSize(Context, (Context.*TablePtr)[Index]); };
 		if (!TryReserveMemory(Context, Algo::TransformAccumulate(IndexBatch, MeasureIndex, uint64(0))))
 		{
 			return -1;
@@ -1317,12 +1394,12 @@ private:
 		return IndexBatch.Num();
 	}
 
-	static void Dispatch(void* VoidContext, TConstArrayView<int32> IndexBatch, IRequestOwner& Owner, const FCopyRequestBatch& Batch)
+	static void Dispatch(void* VoidContext, TConstArrayView<uint32> IndexBatch, IRequestOwner& Owner, const FCopyRequestBatch& Batch)
 	{
 		FCopyContext& Context = *(FCopyContext*)VoidContext;
 
 		TArray<RequestType, TFixedAllocator<GetRequestBatchSize>> Requests;
-		for (int32 Index : IndexBatch)
+		for (uint32 Index : IndexBatch)
 		{
 			const auto& Request = (Context.*TablePtr)[Index];
 			Requests.Add({Request.Name, {Request.Bucket, Request.KeyHash}, ECachePolicy::Query, uint64(Index)});
@@ -1332,14 +1409,14 @@ private:
 		((*Context.Source).*CacheFunctionFor<RequestType>)(Requests, Owner, FOnComplete{Context, Batch});
 	}
 
-	static void Complete(void* VoidContext, TConstArrayView<int32> IndexBatch)
+	static void Complete(void* VoidContext, TConstArrayView<uint32> IndexBatch)
 	{
 		FCopyContext& Context = *(FCopyContext*)VoidContext;
 
 		Context.Progress.GetIndex.fetch_add(IndexBatch.Num(), std::memory_order_relaxed);
 
 		// Calculate the delta between raw and compressed to release/reserve memory now that the size is known.
-		const auto MeasureIndexDelta = [&Context](int32 Index) -> int64
+		const auto MeasureIndexDelta = [&Context](uint32 Index) -> int64
 		{
 			const auto& Request = (Context.*TablePtr)[Index];
 			return int64(MeasureRawSize(Context, Request)) - int64(MeasureCompressedSize(Request));
@@ -1372,7 +1449,7 @@ private:
 
 		void operator()(TCacheResponseFor<RequestType>&& Response) const
 		{
-			const int32 Index = int32(Response.UserData);
+			const uint32 Index = uint32(Response.UserData);
 
 			auto& Request = (Context.*TablePtr)[Index];
 			if (Response.Status != EStatus::Ok)
@@ -1387,7 +1464,7 @@ private:
 				UE_LOGFMT(LogDerivedDataTool, Error, "Value in source missing data for {Key} from '{Name}'.",
 					GetKeyFromResponse(Response), Response.Name);
 			}
-			else if (Context.Find(Index, GetValuesFromResponse(Response)) == -1)
+			else if (Context.Find(Index, GetValuesFromResponse(Response)) == MAX_uint32)
 			{
 				Request.Flags.SourceMask.fetch_or(FCopyFlags::SourceDiffersFlag, std::memory_order_relaxed);
 				UE_LOGFMT(LogDerivedDataTool, Error, "Value in source differs for {Key} from '{Name}'.",
@@ -1412,10 +1489,10 @@ public:
 		return {CanDispatch, Dispatch, Complete, [](void*){}, &Context};
 	}
 
-	static void QueueFromGet(FCopyContext& Context, TConstArrayView<int32> IndexBatch)
+	static void QueueFromGet(FCopyContext& Context, TConstArrayView<uint32> IndexBatch)
 	{
 		// Queue every index that fetched data from the source.
-		for (int32 Index : IndexBatch)
+		for (uint32 Index : IndexBatch)
 		{
 			auto& Request = (Context.*TablePtr)[Index];
 			const uint32 SourceMask = Request.Flags.SourceMask.load(std::memory_order_relaxed);
@@ -1431,7 +1508,7 @@ public:
 	}
 
 private:
-	static uint32 CalculateTargetMask(const FCopyContext& Context, int32 Index)
+	static uint32 CalculateTargetMask(const FCopyContext& Context, uint32 Index)
 	{
 		const auto& Request = (Context.*TablePtr)[Index];
 		const uint32 TargetMissingMask = Request.Flags.TargetMissingMask.load(std::memory_order_relaxed);
@@ -1439,29 +1516,29 @@ private:
 		return TargetMissingMask | (TargetDiffersMask & (Context.Params.bOverwrite ? MAX_uint32 : 0));
 	}
 
-	static int32 CountPutRequests(const FCopyContext& Context, TConstArrayView<int32> IndexBatch)
+	static int32 CountPutRequests(const FCopyContext& Context, TConstArrayView<uint32> IndexBatch)
 	{
 		int32 Count = 0;
-		for (int32 Index : IndexBatch)
+		for (uint32 Index : IndexBatch)
 		{
 			Count += FPlatformMath::CountBits(CalculateTargetMask(Context, Index));
 		}
 		return Count;
 	}
 
-	static int32 CanDispatch(void* VoidContext, TConstArrayView<int32> IndexBatch)
+	static int32 CanDispatch(void* VoidContext, TConstArrayView<uint32> IndexBatch)
 	{
 		FCopyContext& Context = *(FCopyContext*)VoidContext;
 		return CountPutRequests(Context, IndexBatch);
 	}
 
-	static void Dispatch(void* VoidContext, TConstArrayView<int32> IndexBatch, IRequestOwner& Owner, const FCopyRequestBatch& Batch)
+	static void Dispatch(void* VoidContext, TConstArrayView<uint32> IndexBatch, IRequestOwner& Owner, const FCopyRequestBatch& Batch)
 	{
 		FCopyContext& Context = *(FCopyContext*)VoidContext;
 
 		// Calculate which targets to store to from this batch.
 		uint32 ActiveTargetMask = 0;
-		for (int32 Index : IndexBatch)
+		for (uint32 Index : IndexBatch)
 		{
 			ActiveTargetMask |= CalculateTargetMask(Context, Index);
 		}
@@ -1476,7 +1553,7 @@ private:
 			ActiveTargetMask = NextActiveTargetMask;
 
 			TArray<RequestType, TFixedAllocator<PutRequestBatchSize>> Requests;
-			for (int32 Index : IndexBatch)
+			for (uint32 Index : IndexBatch)
 			{
 				const auto& Request = (Context.*TablePtr)[Index];
 				const bool bStoreMissing = !!(Request.Flags.TargetMissingMask.load(std::memory_order_relaxed) & TargetMask);
@@ -1492,16 +1569,16 @@ private:
 		}
 	}
 
-	static void Complete(void* VoidContext, TConstArrayView<int32> IndexBatch)
+	static void Complete(void* VoidContext, TConstArrayView<uint32> IndexBatch)
 	{
 		FCopyContext& Context = *(FCopyContext*)VoidContext;
 
 		Context.Progress.PutIndex.fetch_add(IndexBatch.Num(), std::memory_order_relaxed);
 
-		const auto MeasureIndex = [&Context](int32 Index) -> uint64 { return MeasureCompressedSize((Context.*TablePtr)[Index]); };
+		const auto MeasureIndex = [&Context](uint32 Index) -> uint64 { return MeasureCompressedSize((Context.*TablePtr)[Index]); };
 		ReleaseMemory(Context, Algo::TransformAccumulate(IndexBatch, MeasureIndex, uint64(0)));
 
-		for (int32 Index : IndexBatch)
+		for (uint32 Index : IndexBatch)
 		{
 			(Context.*TablePtr)[Index].Data.Reset();
 		}
@@ -1515,7 +1592,7 @@ private:
 		void operator()(TCacheResponseFor<RequestType>&& Response) const
 		{
 			const int32 TargetIndex = int32(Response.UserData >> 32);
-			const int32 Index = int32(Response.UserData);
+			const uint32 Index = uint32(Response.UserData);
 
 			auto& Request = (Context.*TablePtr)[Index];
 			if (Response.Status != EStatus::Ok && !Context.Params.bDebugSkipPut)
@@ -1642,24 +1719,30 @@ static FCopySummaryByKeyType CalculateBucketSummary(
 	TNotNull<TFunctionWithContext<FCopySummaryValues* (const FCopyContext&, const FCopyFlags&, FCopySummaryByStatus&)>> SelectValues)
 {
 	FCopySummaryByKeyType BucketSummary;
-	for (const TPair<FIoHash, int32>& Pair : Bucket.RecordKeys)
+	for (const FProbingHashTable& RecordKeys : Bucket.RecordKeys)
 	{
-		const FCopyRecord& Record = Context.RecordTable[Pair.Value];
-		if (FCopySummaryValues* RecordSummary = SelectValues(Context, Record.Flags, BucketSummary.Record))
+		for (const uint32 Index : RecordKeys)
 		{
-			++RecordSummary->KeyCount;
-			RecordSummary->ValueCount += Record.ValueCount;
-			RecordSummary->TotalRawSize += MeasureRawSize(Context, Record);
+			const FCopyRecord& Record = Context.RecordTable[Index];
+			if (FCopySummaryValues* RecordSummary = SelectValues(Context, Record.Flags, BucketSummary.Record))
+			{
+				++RecordSummary->KeyCount;
+				RecordSummary->ValueCount += Record.ValueCount;
+				RecordSummary->TotalRawSize += MeasureRawSize(Context, Record);
+			}
 		}
 	}
-	for (const TPair<FIoHash, int32>& Pair : Bucket.ValueKeys)
+	for (const FProbingHashTable& ValueKeys : Bucket.ValueKeys)
 	{
-		const FCopyValue& Value = Context.ValueTable[Pair.Value];
-		if (FCopySummaryValues* ValueSummary = SelectValues(Context, Value.Flags, BucketSummary.Value))
+		for (const uint32 Index : ValueKeys)
 		{
-			++ValueSummary->KeyCount;
-			++ValueSummary->ValueCount;
-			ValueSummary->TotalRawSize += MeasureRawSize(Context, Value);
+			const FCopyValue& Value = Context.ValueTable[Index];
+			if (FCopySummaryValues* ValueSummary = SelectValues(Context, Value.Flags, BucketSummary.Value))
+			{
+				++ValueSummary->KeyCount;
+				++ValueSummary->ValueCount;
+				ValueSummary->TotalRawSize += MeasureRawSize(Context, Value);
+			}
 		}
 	}
 	return BucketSummary;
@@ -1726,7 +1809,7 @@ static void ScanTargets(FCopyContext& Context)
 	for (const TPair<FCacheBucket, TUniquePtr<FCopyBucket>>& BucketPair : Context.Buckets)
 	{
 		const FCopyBucket& Bucket = *BucketPair.Value;
-		Context.Progress.ScanTotal.fetch_add(Bucket.RecordKeys.Num() + Bucket.ValueKeys.Num(), std::memory_order_relaxed);
+		Context.Progress.ScanTotal.fetch_add(Bucket.RecordKeyCount + Bucket.ValueKeyCount, std::memory_order_relaxed);
 	}
 
 	DisplayProgress(Context);
@@ -1736,13 +1819,19 @@ static void ScanTargets(FCopyContext& Context)
 	{
 		const FCacheBucket BucketName = BucketPair.Key;
 		const FCopyBucket& Bucket = *BucketPair.Value;
-		for (const TPair<FIoHash, int32>& Pair : Bucket.RecordKeys)
+		for (const FProbingHashTable& RecordKeys : Bucket.RecordKeys)
 		{
-			Context.ScanRecordQueue.Queue(Pair.Value);
+			for (const uint32 Index : RecordKeys)
+			{
+				Context.ScanRecordQueue.Queue(Index);
+			}
 		}
-		for (const TPair<FIoHash, int32>& Pair : Bucket.ValueKeys)
+		for (const FProbingHashTable& ValueKeys : Bucket.ValueKeys)
 		{
-			Context.ScanValueQueue.Queue(Pair.Value);
+			for (const uint32 Index : ValueKeys)
+			{
+				Context.ScanValueQueue.Queue(Index);
+			}
 		}
 	}
 
@@ -1767,17 +1856,21 @@ static void QueueForceOverwrite(FCopyContext& Context)
 	{
 		const FCacheBucket BucketName = BucketPair.Key;
 		const FCopyBucket& Bucket = *BucketPair.Value;
-		for (const TPair<FIoHash, int32>& Pair : Bucket.RecordKeys)
+		for (const FProbingHashTable& RecordKeys : Bucket.RecordKeys)
 		{
-			const int32 Index = Pair.Value;
-			Context.RecordTable[Index].Flags.TargetDiffersMask.store(TargetMask, std::memory_order_relaxed);
-			FGetRecordDispatcher::QueueFromScan(Context, MakeArrayView(&Index, 1));
+			for (const uint32 Index : RecordKeys)
+			{
+				Context.RecordTable[Index].Flags.TargetDiffersMask.store(TargetMask, std::memory_order_relaxed);
+				FGetRecordDispatcher::QueueFromScan(Context, MakeArrayView(&Index, 1));
+			}
 		}
-		for (const TPair<FIoHash, int32>& Pair : Bucket.ValueKeys)
+		for (const FProbingHashTable& ValueKeys : Bucket.ValueKeys)
 		{
-			const int32 Index = Pair.Value;
-			Context.ValueTable[Index].Flags.TargetDiffersMask.store(TargetMask, std::memory_order_relaxed);
-			FGetValueDispatcher::QueueFromScan(Context, MakeArrayView(&Index, 1));
+			for (const uint32 Index : ValueKeys)
+			{
+				Context.ValueTable[Index].Flags.TargetDiffersMask.store(TargetMask, std::memory_order_relaxed);
+				FGetValueDispatcher::QueueFromScan(Context, MakeArrayView(&Index, 1));
+			}
 		}
 	}
 
