DerivedDataTool: Fixed a build warning-as-error with Clang

#rnx
#Synced-CL 50967057

diff --git a/Engine/Source/Programs/DerivedDataTool/Private/DerivedDataWorkerCommand.cpp b/Engine/Source/Programs/DerivedDataTool/Private/DerivedDataWorkerCommand.cpp
--- a/Engine/Source/Programs/DerivedDataTool/Private/DerivedDataWorkerCommand.cpp
+++ b/Engine/Source/Programs/DerivedDataTool/Private/DerivedDataWorkerCommand.cpp
@@ -100,28 +100,28 @@ struct FDisplaySize
 {
 	uint64 Size = 0;
 
-	friend void SerializeForLog(FCbWriter& Writer, const FDisplaySize& Size)
+	friend void SerializeForLog(FCbWriter& Writer, const FDisplaySize& DisplaySize)
 	{
 		TUtf8StringBuilder<64> Text;
-		if (Size.Size < uint64(10) * 1024)
+		if (DisplaySize.Size < uint64(10) * 1024)
 		{
-			Text << (Size.Size) << " B";
+			Text << (DisplaySize.Size) << " B";
 		}
-		else if (Size.Size < uint64(10) * 1024 * 1024)
+		else if (DisplaySize.Size < uint64(10) * 1024 * 1024)
 		{
-			Text << (Size.Size / 1024) << " KiB";
+			Text << (DisplaySize.Size / 1024) << " KiB";
 		}
-		else if (Size.Size < uint64(10) * 1024 * 1024 * 1024)
+		else if (DisplaySize.Size < uint64(10) * 1024 * 1024 * 1024)
 		{
-			Text << (Size.Size / 1024 / 1024) << " MiB";
+			Text << (DisplaySize.Size / 1024 / 1024) << " MiB";
 		}
 		else
 		{
-			Text << (Size.Size / 1024 / 1024 / 1024) << " GiB";
+			Text << (DisplaySize.Size / 1024 / 1024 / 1024) << " GiB";
 		}
 
 		Writer.BeginObject();
-		Writer.AddInteger("Size", Size.Size);
+		Writer.AddInteger("Size", DisplaySize.Size);
 		Writer.AddString("$text", Text);
 		Writer.EndObject();
 	}
