DerivedDataTool: Added the build worker package path to the log output

#rnx
#Synced-CL 50967057

diff --git a/Engine/Source/Programs/DerivedDataTool/Private/DerivedDataWorkerCommand.cpp b/Engine/Source/Programs/DerivedDataTool/Private/DerivedDataWorkerCommand.cpp
--- a/Engine/Source/Programs/DerivedDataTool/Private/DerivedDataWorkerCommand.cpp
+++ b/Engine/Source/Programs/DerivedDataTool/Private/DerivedDataWorkerCommand.cpp
@@ -152,10 +152,12 @@ void DisplayWorker(const FBuildWorker& Worker, bool bListFiles)
 	int64 TotalCompressedSize = IFileManager::Get().FileSize(*WriteToString<256>(Worker.GetPackagePath()));
 
 	UE_LOGFMT(LogDerivedDataTool, Display, "");
-	UE_LOGFMT(LogDerivedDataTool, Display, "Build Worker: {WorkerName}", Worker.GetName());
-	UE_LOGFMT(LogDerivedDataTool, Display, "Build Worker Key: {WorkerKey}", Worker.GetKey().Hash);
+	UE_CLOGFMT(!Worker.GetPackagePath().IsEmpty(),
+		LogDerivedDataTool, Display, "Build Worker Package: {PackagePath}", Worker.GetPackagePath());
+	UE_LOGFMT(LogDerivedDataTool, Display, "Build Worker Name: {WorkerName}", Worker.GetName());
+	UE_LOGFMT(LogDerivedDataTool, Display, "Build Worker Host: {HostPlatform}", Worker.GetHostPlatform());
 	UE_LOGFMT(LogDerivedDataTool, Display, "Build Worker Path: {WorkerPath}", Worker.GetPath());
-	UE_LOGFMT(LogDerivedDataTool, Display, "Build Worker Host Platform: {HostPlatform}", Worker.GetHostPlatform());
+	UE_LOGFMT(LogDerivedDataTool, Display, "Build Worker Key: {WorkerKey}", Worker.GetKey().Hash);
 	UE_LOGFMT(LogDerivedDataTool, Display, "Build System Version: {BuildSystemVersion}", Worker.GetBuildSystemVersion());
 	UE_LOGFMT(LogDerivedDataTool, Display,
 		"Build Worker Summary: {FileCount} files, {RawSize} uncompressed, {CompressedSize} compressed",
